Use setlocale to be portable for all locales

diff --git a/Engine/Source/Runtime/Unix/UnixCommonStartup/Private/UnixCommonStartup.cpp b/Engine/Source/Runtime/Unix/UnixCommonStartup/Private/UnixCommonStartup.cpp
--- a/Engine/Source/Runtime/Unix/UnixCommonStartup/Private/UnixCommonStartup.cpp
+++ b/Engine/Source/Runtime/Unix/UnixCommonStartup/Private/UnixCommonStartup.cpp
@@ -183,8 +183,8 @@ int CommonUnixMain(int argc, char *argv[], int (*RealMain)(const TCHAR * Command
 		int ErrNo = errno;
 		fprintf(stderr, "Unable to setenv(LC_NUMERIC): errno=%d (%s)", ErrNo, strerror(ErrNo));
 	}
-	setlocale(LC_CTYPE, "");
 
+	setlocale(LC_ALL, "");
 
 	for (int32 Option = 1; Option < argc; Option++)
 	{
