Merging Arciel Rekman's removal of setlocale() for Linux from Dev-Console

diff --git a/Engine/Source/Programs/CrashReportClient/Private/Linux/CrashReportClientMainLinux.cpp b/Engine/Source/Programs/CrashReportClient/Private/Linux/CrashReportClientMainLinux.cpp
--- a/Engine/Source/Programs/CrashReportClient/Private/Linux/CrashReportClientMainLinux.cpp
+++ b/Engine/Source/Programs/CrashReportClient/Private/Linux/CrashReportClientMainLinux.cpp
@@ -51,7 +51,6 @@ int main(int argc, const char *argv[])
 
 	FString SavedCommandLine;
 
-	setlocale(LC_CTYPE, "");
 	for (int32 Option = 1; Option < argc; Option++)
 	{
 		SavedCommandLine += TEXT(" ");
diff --git a/Engine/Source/Programs/UnrealWatchdog/Private/Linux/UnrealWatchdogMainLinux.cpp b/Engine/Source/Programs/UnrealWatchdog/Private/Linux/UnrealWatchdogMainLinux.cpp
--- a/Engine/Source/Programs/UnrealWatchdog/Private/Linux/UnrealWatchdogMainLinux.cpp
+++ b/Engine/Source/Programs/UnrealWatchdog/Private/Linux/UnrealWatchdogMainLinux.cpp
@@ -44,7 +44,6 @@ int main(int argc, const char *argv[])
 
 	FString SavedCommandLine;
 
-	setlocale(LC_CTYPE, "");
 	for (int32 Option = 1; Option < argc; Option++)
 	{
 		SavedCommandLine += TEXT(" ");
diff --git a/Engine/Source/Runtime/Unix/UnixCommonStartup/Private/UnixCommonStartup.cpp b/Engine/Source/Runtime/Unix/UnixCommonStartup/Private/UnixCommonStartup.cpp
--- a/Engine/Source/Runtime/Unix/UnixCommonStartup/Private/UnixCommonStartup.cpp
+++ b/Engine/Source/Runtime/Unix/UnixCommonStartup/Private/UnixCommonStartup.cpp
@@ -188,8 +188,6 @@ int CommonUnixMain(int argc, char *argv[], int (*RealMain)(const TCHAR * Command
 		fprintf(stderr, "Unable to setenv(LC_NUMERIC): errno=%d (%s)", ErrNo, strerror(ErrNo));
 	}
 
-	setlocale(LC_ALL, "");
-
 	for (int32 Option = 1; Option < argc; Option++)
 	{
 		GSavedCommandLine += TEXT(" ");
