Also set the SetCurrentThreadAsMasterThread in the CommonCrashHandler

diff --git a/Engine/Source/Runtime/Unix/UnixCommonStartup/Private/UnixCommonStartup.cpp b/Engine/Source/Runtime/Unix/UnixCommonStartup/Private/UnixCommonStartup.cpp
--- a/Engine/Source/Runtime/Unix/UnixCommonStartup/Private/UnixCommonStartup.cpp
+++ b/Engine/Source/Runtime/Unix/UnixCommonStartup/Private/UnixCommonStartup.cpp
@@ -32,6 +32,7 @@ void CommonUnixCrashHandler(const FGenericCrashContext& GenericContext)
 	const_cast< FUnixCrashContext& >(Context).CaptureStackTrace();
 	if (GLog)
 	{
+		GLog->SetCurrentThreadAsMasterThread();
 		GLog->Flush();
 	}
 	if (GWarn)
