Rework UnixCommonStartup to avoid relying on a dangling symbol.

diff --git a/Engine/Source/Runtime/Launch/Private/Unix/LaunchUnix.cpp b/Engine/Source/Runtime/Launch/Private/Unix/LaunchUnix.cpp
--- a/Engine/Source/Runtime/Launch/Private/Unix/LaunchUnix.cpp
+++ b/Engine/Source/Runtime/Launch/Private/Unix/LaunchUnix.cpp
@@ -12,15 +12,13 @@ extern int32 GuardedMain( const TCHAR* CmdLine );
  * Other platforms call FEngineLoop::AppExit() in their main() (removed by preprocessor if compiled without engine), but on Unix we want to share a common main() in CommonUnixStartup module,
  * so not just the engine but all the programs could share this logic. Unfortunately, AppExit() practice breaks this nice approach since FEngineLoop cannot be moved outside of Launch without
  * making too many changes. Hence CommonUnixMain will call it through this function if WITH_ENGINE is defined.
- *
- * If you change the prototype here, update CommonUnixMain() too!
  */
-void LAUNCH_API LaunchUnix_FEngineLoop_AppExit()
+void LaunchUnix_FEngineLoop_AppExit()
 {
 	return FEngineLoop::AppExit();
 }
 
 int main(int argc, char *argv[])
 {
-	return CommonUnixMain(argc, argv, &GuardedMain);
+	return CommonUnixMain(argc, argv, &GuardedMain, &LaunchUnix_FEngineLoop_AppExit);
 }
diff --git a/Engine/Source/Runtime/Unix/UnixCommonStartup/Private/UnixCommonStartup.cpp b/Engine/Source/Runtime/Unix/UnixCommonStartup/Private/UnixCommonStartup.cpp
--- a/Engine/Source/Runtime/Unix/UnixCommonStartup/Private/UnixCommonStartup.cpp
+++ b/Engine/Source/Runtime/Unix/UnixCommonStartup/Private/UnixCommonStartup.cpp
@@ -18,11 +18,6 @@ static FString GSavedCommandLine;
 extern int32 GuardedMain( const TCHAR* CmdLine );
 extern void LaunchStaticShutdownAfterError();
 
-#if WITH_ENGINE
-// see comment in LaunchUnix.cpp for details why it is done this way
-extern void LaunchUnix_FEngineLoop_AppExit();
-#endif // WITH_ENGINE
-
 /**
  * Game-specific crash reporter
  */
@@ -166,7 +161,7 @@ static bool IncreasePerProcessLimits()
 	return true;
 }
 
-int CommonUnixMain(int argc, char *argv[], int (*RealMain)(const TCHAR * CommandLine))
+int CommonUnixMain(int argc, char *argv[], int (*RealMain)(const TCHAR * CommandLine), void (*AppExitCallback)())
 {
 	FString EarlyInitCommandLine;
 	FPlatformApplicationMisc::EarlyUnixInitialization(EarlyInitCommandLine);
@@ -246,9 +241,17 @@ int CommonUnixMain(int argc, char *argv[], int (*RealMain)(const TCHAR * Command
 	}
 
 	// Final shut down.
-#if WITH_ENGINE
-	LaunchUnix_FEngineLoop_AppExit();
-#endif // WITH_ENGINE
+	if (AppExitCallback)
+	{
+		// Workaround function to avoid circular dependencies between Launch and CommonUnixStartup modules.
+
+		// Other platforms call FEngineLoop::AppExit() in their main() (removed by preprocessor if compiled without engine), 
+		// but on Unix we want to share a common main() in CommonUnixStartup module, so not just the engine but all the programs 
+		// could share this logic. Unfortunately, AppExit() practice breaks this nice approach since FEngineLoop cannot be moved outside of 
+		// Launch module without making too many changes. Hence CommonUnixMain will call it through this function if provided.
+
+		AppExitCallback();
+	}
 
 	// check if a specific return code has been set
 	uint8 OverriddenErrorLevel = 0;
diff --git a/Engine/Source/Runtime/Unix/UnixCommonStartup/Public/UnixCommonStartup.h b/Engine/Source/Runtime/Unix/UnixCommonStartup/Public/UnixCommonStartup.h
--- a/Engine/Source/Runtime/Unix/UnixCommonStartup/Public/UnixCommonStartup.h
+++ b/Engine/Source/Runtime/Unix/UnixCommonStartup/Public/UnixCommonStartup.h
@@ -9,6 +9,7 @@
  * @param argc - number of arguments in argv[]
  * @param argv - array of arguments
  * @param RealMain - the next main routine to call in chain
+ * @param AppExitCallback - workaround for Launch module that needs to call FEngineLoop::AppExit() at certain point
  * @return error code to return to the OS
  */
-int UNIXCOMMONSTARTUP_API CommonUnixMain(int argc, char *argv[], int (*RealMain)(const TCHAR * CommandLine));
+int UNIXCOMMONSTARTUP_API CommonUnixMain(int argc, char *argv[], int (*RealMain)(const TCHAR * CommandLine), void (*AppExitCallback)() = nullptr);
