Add missing dependency to UnixCommonStartup.

- It calls EarlyUnixInitialization.

diff --git a/Engine/Source/Runtime/Unix/UnixCommonStartup/UnixCommonStartup.Build.cs b/Engine/Source/Runtime/Unix/UnixCommonStartup/UnixCommonStartup.Build.cs
--- a/Engine/Source/Runtime/Unix/UnixCommonStartup/UnixCommonStartup.Build.cs
+++ b/Engine/Source/Runtime/Unix/UnixCommonStartup/UnixCommonStartup.Build.cs
@@ -8,6 +8,7 @@ public class UnixCommonStartup : ModuleRules
 	public UnixCommonStartup(ReadOnlyTargetRules Target) : base(Target)
 	{
 		PrivateDependencyModuleNames.Add("Core");
+		PrivateDependencyModuleNames.Add("ApplicationCore");
 		PrivateDependencyModuleNames.Add("Projects");
 
 		PrivateDependencyModuleNames.Add("Slate");
