Set the max allowed files open as our default value

diff --git a/Engine/Source/Runtime/Unix/UnixCommonStartup/Private/UnixCommonStartup.cpp b/Engine/Source/Runtime/Unix/UnixCommonStartup/Private/UnixCommonStartup.cpp
--- a/Engine/Source/Runtime/Unix/UnixCommonStartup/Private/UnixCommonStartup.cpp
+++ b/Engine/Source/Runtime/Unix/UnixCommonStartup/Private/UnixCommonStartup.cpp
@@ -99,6 +99,23 @@ bool SetResourceLimit(int Resource, rlim_t DesiredLimit, bool bIncreaseOnly)
 	return true;
 }
 
+/**
+ * Sets (hard) limit on a specific resource
+ *
+ * @param Resource - one of RLIMIT_* values
+ */
+static bool SetResourceMaxHardLimit(int Resource)
+{
+	rlimit Limit;
+	if (getrlimit(Resource, &Limit) != 0)
+	{
+		fprintf(stderr, "getrlimit() failed with error %d (%s)\n", errno, strerror(errno));
+		return false;
+	}
+
+	return SetResourceLimit(Resource, Limit.rlim_max, true);
+}
+
 /**
  * Expects GSavedCommandLine to be set up. Increases limit on
  *  - number of open files to be no less than desired (if specified on command line, otherwise left alone)
@@ -107,7 +124,7 @@ bool SetResourceLimit(int Resource, rlim_t DesiredLimit, bool bIncreaseOnly)
  */
 static bool IncreasePerProcessLimits()
 {
-	// honor the parameter if given, but don't change limits if not
+	// honor the parameter if given, else change the limit of open files to the max hard limit allowed
 	int32 FileHandlesToReserve = -1;
 	if (FParse::Value(*GSavedCommandLine, TEXT("numopenfiles="), FileHandlesToReserve) && FileHandlesToReserve > 0)
 	{
@@ -122,6 +139,11 @@ static bool IncreasePerProcessLimits()
 			return false;
 		}
 	}
+	else
+	{
+		// Set the highest value we can for number of files open
+		SetResourceMaxHardLimit(RLIMIT_NOFILE);
+	}
 
 	// core dump policy:
 	// - Shipping disables it by default (unless -core is passed)
