Linux compile fix.

#rnx

diff --git a/Engine/Source/Runtime/Unix/UnixCommonStartup/Private/UnixCommonStartup.cpp b/Engine/Source/Runtime/Unix/UnixCommonStartup/Private/UnixCommonStartup.cpp
--- a/Engine/Source/Runtime/Unix/UnixCommonStartup/Private/UnixCommonStartup.cpp
+++ b/Engine/Source/Runtime/Unix/UnixCommonStartup/Private/UnixCommonStartup.cpp
@@ -32,7 +32,7 @@ void CommonUnixCrashHandler(const FGenericCrashContext& GenericContext)
 	FGenericCrashContext::SetMemoryStats(FPlatformMemory::GetStats());
 
 	// better than having mutable fields?
-	const_cast< FUnixCrashContext& >(Context).CaptureStackTrace(Context.ProgramCounter);
+	const_cast< FUnixCrashContext& >(Context).CaptureStackTrace(Context.ErrorFrame);
 	if (GLog)
 	{
 		GLog->SetCurrentThreadAsMasterThread();
