Disable ASan leak detector by default. While it does detect leaks we have many many leaks that are not proper in UEs term which mean we ignore all the current leaks.

Need to look into a suppression file for these so we can re-enable this

diff --git a/Engine/Source/Programs/UnrealBuildTool/Platform/Linux/LinuxToolChain.cs b/Engine/Source/Programs/UnrealBuildTool/Platform/Linux/LinuxToolChain.cs
--- a/Engine/Source/Programs/UnrealBuildTool/Platform/Linux/LinuxToolChain.cs
+++ b/Engine/Source/Programs/UnrealBuildTool/Platform/Linux/LinuxToolChain.cs
@@ -682,7 +682,7 @@ namespace UnrealBuildTool
 			{
 				Result += " -Wno-unused-value";
 				if (!bRetainFramePointers)
-                {
+				{
 					Result += " -fomit-frame-pointer";
 				}
 			}
@@ -729,6 +729,9 @@ namespace UnrealBuildTool
 				if (Options.HasFlag(LinuxToolChainOptions.EnableAddressSanitizer) || Options.HasFlag(LinuxToolChainOptions.EnableMemorySanitizer))
 				{
 					Result += " -O1 -g";
+
+					// This enables __asan_default_options() in UnixCommonStartup.h which disables the leak detector
+					Result += " -DDISABLE_ASAN_LEAK_DETECTOR=1";
 				}
 				else if (Options.HasFlag(LinuxToolChainOptions.EnableThreadSanitizer))
 				{
diff --git a/Engine/Source/Runtime/Unix/UnixCommonStartup/Public/UnixCommonStartup.h b/Engine/Source/Runtime/Unix/UnixCommonStartup/Public/UnixCommonStartup.h
--- a/Engine/Source/Runtime/Unix/UnixCommonStartup/Public/UnixCommonStartup.h
+++ b/Engine/Source/Runtime/Unix/UnixCommonStartup/Public/UnixCommonStartup.h
@@ -4,6 +4,10 @@
 
 #include "CoreMinimal.h"
 
+#ifndef DISABLE_ASAN_LEAK_DETECTOR
+	#define DISABLE_ASAN_LEAK_DETECTOR 0
+#endif
+
 /**
  * @brief CommonUnixMain - executes common startup code for Unix programs/engine
  * @param argc - number of arguments in argv[]
@@ -13,3 +17,14 @@
  * @return error code to return to the OS
  */
 int UNIXCOMMONSTARTUP_API CommonUnixMain(int argc, char *argv[], int (*RealMain)(const TCHAR * CommandLine), void (*AppExitCallback)() = nullptr);
+
+#if DISABLE_ASAN_LEAK_DETECTOR
+/*
+ * We honestly leak so much this output is not super useful, so lets disable by default but if you want to re-enable disable
+ * this DEFINE in LinuxToolchain.cs area
+ */
+extern "C" const char* UNIXCOMMONSTARTUP_API __asan_default_options()
+{
+	return "detect_leaks=0";
+}
+#endif
