Fix duplicate symbols when compiling monolithic for this function

diff --git a/Engine/Source/Runtime/Unix/UnixCommonStartup/Private/UnixCommonStartup.cpp b/Engine/Source/Runtime/Unix/UnixCommonStartup/Private/UnixCommonStartup.cpp
--- a/Engine/Source/Runtime/Unix/UnixCommonStartup/Private/UnixCommonStartup.cpp
+++ b/Engine/Source/Runtime/Unix/UnixCommonStartup/Private/UnixCommonStartup.cpp
@@ -18,6 +18,11 @@ static FString GSavedCommandLine;
 extern int32 GuardedMain( const TCHAR* CmdLine );
 extern void LaunchStaticShutdownAfterError();
 
+const char* __asan_default_options()
+{
+	return "detect_leaks=0";
+}
+
 /**
  * Game-specific crash reporter
  */
diff --git a/Engine/Source/Runtime/Unix/UnixCommonStartup/Public/UnixCommonStartup.h b/Engine/Source/Runtime/Unix/UnixCommonStartup/Public/UnixCommonStartup.h
--- a/Engine/Source/Runtime/Unix/UnixCommonStartup/Public/UnixCommonStartup.h
+++ b/Engine/Source/Runtime/Unix/UnixCommonStartup/Public/UnixCommonStartup.h
@@ -23,8 +23,5 @@ int UNIXCOMMONSTARTUP_API CommonUnixMain(int argc, char *argv[], int (*RealMain)
  * We honestly leak so much this output is not super useful, so lets disable by default but if you want to re-enable disable
  * this DEFINE in LinuxToolchain.cs area
  */
-extern "C" const char* UNIXCOMMONSTARTUP_API __asan_default_options()
-{
-	return "detect_leaks=0";
-}
+extern "C" const char* UNIXCOMMONSTARTUP_API __asan_default_options();
 #endif
