Move __asan_default_options to the main exe as it needs to be bound before asan*.so is loaded to override that function

diff --git a/Engine/Source/Runtime/Launch/Private/Unix/LaunchUnix.cpp b/Engine/Source/Runtime/Launch/Private/Unix/LaunchUnix.cpp
--- a/Engine/Source/Runtime/Launch/Private/Unix/LaunchUnix.cpp
+++ b/Engine/Source/Runtime/Launch/Private/Unix/LaunchUnix.cpp
@@ -4,6 +4,21 @@
 #include "LaunchEngineLoop.h"
 #include "UnixCommonStartup.h"
 
+#ifndef DISABLE_ASAN_LEAK_DETECTOR
+	#define DISABLE_ASAN_LEAK_DETECTOR 0
+#endif
+
+#if DISABLE_ASAN_LEAK_DETECTOR
+/*
+ * We honestly leak so much this output is not super useful, so lets disable by default but if you want to re-enable disable
+ * this DEFINE in LinuxToolchain.cs area. This must be defined in the main binary otherwise asan*.so will bind to this symbol first
+ */
+extern "C" const char* LAUNCH_API __asan_default_options()
+{
+	return "detect_leaks=0";
+}
+#endif
+
 extern int32 GuardedMain( const TCHAR* CmdLine );
 
 /**
diff --git a/Engine/Source/Runtime/Unix/UnixCommonStartup/Private/UnixCommonStartup.cpp b/Engine/Source/Runtime/Unix/UnixCommonStartup/Private/UnixCommonStartup.cpp
--- a/Engine/Source/Runtime/Unix/UnixCommonStartup/Private/UnixCommonStartup.cpp
+++ b/Engine/Source/Runtime/Unix/UnixCommonStartup/Private/UnixCommonStartup.cpp
@@ -18,11 +18,6 @@ static FString GSavedCommandLine;
 extern int32 GuardedMain( const TCHAR* CmdLine );
 extern void LaunchStaticShutdownAfterError();
 
-const char* __asan_default_options()
-{
-	return "detect_leaks=0";
-}
-
 /**
  * Game-specific crash reporter
  */
diff --git a/Engine/Source/Runtime/Unix/UnixCommonStartup/Public/UnixCommonStartup.h b/Engine/Source/Runtime/Unix/UnixCommonStartup/Public/UnixCommonStartup.h
--- a/Engine/Source/Runtime/Unix/UnixCommonStartup/Public/UnixCommonStartup.h
+++ b/Engine/Source/Runtime/Unix/UnixCommonStartup/Public/UnixCommonStartup.h
@@ -4,10 +4,6 @@
 
 #include "CoreMinimal.h"
 
-#ifndef DISABLE_ASAN_LEAK_DETECTOR
-	#define DISABLE_ASAN_LEAK_DETECTOR 0
-#endif
-
 /**
  * @brief CommonUnixMain - executes common startup code for Unix programs/engine
  * @param argc - number of arguments in argv[]
@@ -17,11 +13,3 @@
  * @return error code to return to the OS
  */
 int UNIXCOMMONSTARTUP_API CommonUnixMain(int argc, char *argv[], int (*RealMain)(const TCHAR * CommandLine), void (*AppExitCallback)() = nullptr);
-
-#if DISABLE_ASAN_LEAK_DETECTOR
-/*
- * We honestly leak so much this output is not super useful, so lets disable by default but if you want to re-enable disable
- * this DEFINE in LinuxToolchain.cs area
- */
-extern "C" const char* UNIXCOMMONSTARTUP_API __asan_default_options();
-#endif
