Add reading commands/executing from stdin if -cmdstdin is passed on linux command line

#rnx

diff --git a/Engine/Source/Runtime/Unix/UnixCommonStartup/Private/UnixCommonStartup.cpp b/Engine/Source/Runtime/Unix/UnixCommonStartup/Private/UnixCommonStartup.cpp
--- a/Engine/Source/Runtime/Unix/UnixCommonStartup/Private/UnixCommonStartup.cpp
+++ b/Engine/Source/Runtime/Unix/UnixCommonStartup/Private/UnixCommonStartup.cpp
@@ -1,6 +1,9 @@
 // Copyright Epic Games, Inc. All Rights Reserved.
 
 #include "UnixCommonStartup.h"
+#if WITH_ENGINE
+#include "Engine/Engine.h"
+#endif // WITH_ENGINE
 #include "Misc/OutputDeviceRedirector.h"
 #include "Misc/OutputDeviceError.h"
 #include "Misc/FeedbackContext.h"
@@ -18,6 +21,49 @@ static FString GSavedCommandLine;
 extern int32 GuardedMain( const TCHAR* CmdLine );
 extern void LaunchStaticShutdownAfterError();
 
+#if !UE_BUILD_SHIPPING && WITH_ENGINE
+namespace UnixStdinTerminal
+{
+	FTSTicker::FDelegateHandle CurrentTicker;
+
+	bool Tick(float DeltaTime)
+	{
+		static bool bInited = false;
+
+		if (!bInited)
+		{
+			bInited = true;
+			int flags = fcntl(STDIN_FILENO, F_GETFL);
+			flags |= O_NONBLOCK;
+			fcntl(STDIN_FILENO, F_SETFL, flags);
+		}
+
+		int c = getchar();
+		if (c != EOF)
+		{
+			int CmdBufferLen = 0;
+			char CmdBuffer[4096];
+
+			CmdBuffer[CmdBufferLen++] = c;
+			while (((c = getchar()) != EOF) && (CmdBufferLen < 4094))
+			{
+				if (c == '\n' || c == '\r')
+					break;
+
+				CmdBuffer[CmdBufferLen++] = c;
+			}
+
+			CmdBuffer[CmdBufferLen] = 0;
+
+			UE_LOG(LogCore, Log, TEXT("Execute Command:'%s'\n"), ANSI_TO_TCHAR(CmdBuffer));
+			GEngine->DeferredCommands.Add(ANSI_TO_TCHAR(CmdBuffer));
+		}
+
+		return true;
+	}
+}
+#endif //!UE_BUILD_SHIPPING && WITH_ENGINE
+
 /**
  * Game-specific crash reporter
  */
@@ -263,11 +309,26 @@ int CommonUnixMain(int argc, char *argv[], int (*RealMain)(const TCHAR * Command
 		}
 		else
 		{
+#if !UE_BUILD_SHIPPING && WITH_ENGINE
+			bool bReadCommandsFromStdin = FParse::Param(*GSavedCommandLine, TEXT("cmdstdin"));
+			if (bReadCommandsFromStdin)
+			{
+				UnixStdinTerminal::CurrentTicker = FTSTicker::GetCoreTicker().AddTicker(FTickerDelegate::CreateStatic(&UnixStdinTerminal::Tick), 1.0f);
+			}
+#endif
+
 			FPlatformMisc::SetCrashHandler(CommonUnixCrashHandler);
 			GIsGuarded = 1;
 			// Run the guarded code.
 			ErrorLevel = RealMain( *GSavedCommandLine );
 			GIsGuarded = 0;
+
+#if !UE_BUILD_SHIPPING && WITH_ENGINE
+			if (bReadCommandsFromStdin)
+			{
+				FTSTicker::GetCoreTicker().RemoveTicker(UnixStdinTerminal::CurrentTicker);
+			}
+#endif
 		}
 	}
 
diff --git a/Engine/Source/Runtime/Unix/UnixCommonStartup/UnixCommonStartup.Build.cs b/Engine/Source/Runtime/Unix/UnixCommonStartup/UnixCommonStartup.Build.cs
--- a/Engine/Source/Runtime/Unix/UnixCommonStartup/UnixCommonStartup.Build.cs
+++ b/Engine/Source/Runtime/Unix/UnixCommonStartup/UnixCommonStartup.Build.cs
@@ -13,6 +13,12 @@ public class UnixCommonStartup : ModuleRules
 
 		PrivateDependencyModuleNames.Add("Slate");
 
+		// Only include Engine if necessary
+		if (Target.bCompileAgainstEngine)
+		{
+			PublicDependencyModuleNames.Add("Engine");
+		}
+
 		if (Target.Type == TargetType.Editor)
 		{
 			PrivateDependencyModuleNames.AddRange(
